[Insights] Added bAllowPanningOnScreenEdges in DefaultUnrealInsightsSettings.ini file.

diff --git a/Engine/Programs/UnrealInsights/Config/DefaultUnrealInsightsSettings.ini b/Engine/Programs/UnrealInsights/Config/DefaultUnrealInsightsSettings.ini
--- a/Engine/Programs/UnrealInsights/Config/DefaultUnrealInsightsSettings.ini
+++ b/Engine/Programs/UnrealInsights/Config/DefaultUnrealInsightsSettings.ini
@@ -19,6 +19,8 @@ AutoScrollFrameAlignment=game
 AutoScrollViewportOffsetPercent=0.1
 ; Minimum time between two auto-scroll updates, in [seconds].
 AutoScrollMinDelay=0.3
+; If enabled, the panning is allowed to continue when the mouse cursor reaches the edges of the screen.
+bAllowPanningOnScreenEdges=False
 
 [Insights.MemoryProfiler]
 ; Add symbol search paths by adding to this array. Order is important.
