[Insights] Fix persistent settings.

diff --git a/Engine/Programs/UnrealInsights/Config/DefaultUnrealInsightsSessionBrowserSettings.ini b/Engine/Programs/UnrealInsights/Config/DefaultUnrealInsightsSessionBrowserSettings.ini
--- a/Engine/Programs/UnrealInsights/Config/DefaultUnrealInsightsSessionBrowserSettings.ini
+++ b/Engine/Programs/UnrealInsights/Config/DefaultUnrealInsightsSessionBrowserSettings.ini
@@ -1,3 +1,6 @@
+[SectionsToSave]
+bCanSaveAllSections=true
+
 [Insights.SessionBrowser]
 ; Whether Insights should signal to an UE application to auto-connect and start tracing when Insights is running.
 AutoConnect=True
diff --git a/Engine/Programs/UnrealInsights/Config/DefaultUnrealInsightsSettings.ini b/Engine/Programs/UnrealInsights/Config/DefaultUnrealInsightsSettings.ini
--- a/Engine/Programs/UnrealInsights/Config/DefaultUnrealInsightsSettings.ini
+++ b/Engine/Programs/UnrealInsights/Config/DefaultUnrealInsightsSettings.ini
@@ -1,3 +1,6 @@
+[SectionsToSave]
+bCanSaveAllSections=true
+
 [Insights.TimingProfiler]
 ; The default (initial) zoom level of the Timing view.
 ; This option can be changed from the "Tracks" drop down menu.
