[Insights] Updated comments and examples in the DefaultUnrealInsightsSettings.ini.

diff --git a/Engine/Programs/UnrealInsights/Config/DefaultUnrealInsightsSettings.ini b/Engine/Programs/UnrealInsights/Config/DefaultUnrealInsightsSettings.ini
--- a/Engine/Programs/UnrealInsights/Config/DefaultUnrealInsightsSettings.ini
+++ b/Engine/Programs/UnrealInsights/Config/DefaultUnrealInsightsSettings.ini
@@ -6,14 +6,16 @@ bCanSaveAllSections=True
 +Name=LoadingProfiler
 +Name=MemoryProfiler
 
-[Insights.Symbols]
+[Insights]
+
+; Search path used to locate a source file, if a file is not found using its full file path.
+;SourceFilesSearchPath=D:\ue5-main\
 
 ; Add symbol search paths by adding to this array. Order is important.
 ;+SymbolSearchPaths=C:\Symbols
 ; Symbol server format is also accepted, such as the official Microsoft symbol servers (see https://docs.microsoft.com/en-us/windows-hardware/drivers/debugger/advanced-symsrv-use for details)
 ;+SymbolSearchPaths=srv*C:\Symbols;https://msdl.microsoft.com/download/symbols
 
-
 [Insights.TimingProfiler]
 
 ; The default (initial) zoom level of the Timing view.
@@ -147,18 +149,18 @@ Mode=2
 ; Grouping Mode for the timers panel. 0 : Flat; 1 : ByName; 3 : ByType; 4 : ByInstanceCount
 GroupingMode=3
 
-; Controls the CPU filter button from Timers View.
+; Controls the GPU filter button from Timers View (GPU Scope timers).
+ShowGpuTimers=True
+
+; Controls the CPU filter button from Timers View (CPU Scope timers).
 ShowCpuTimers=True
 
-; Controls the CPU Sampling filter button from Timers View.
+; Controls the CPU Sampling filter button from Timers View (CPU Stack Sampling timers).
 ShowCpuSamplingTimers=True
 
-; Controls the Verse filter button from Timers View.
+; Controls the Verse filter button from Timers View (Verse Sampling timers).
 ShowVerseSamplingTimers=True
 
-; Controls the GPU filter button from Timers View.
-ShowGpuTimers=True
-
 ; Controls the "Filter Out Zero Count Timers" button from Timers View.
 ShowZeroCountTimers=True
 
@@ -187,7 +189,7 @@ ShowNonTargetProcessEvents=True
 
 [Insights.MemoryProfiler]
 
-; Memory Budgets (used by the LLM Tags view)
+; Memory Budgets (used by the Memory Tags view)
 +BudgetFilePath=(Label="Default",Path="Engine/Programs/UnrealInsights/Config/DefaultMemoryBudgets.xml")
 
 [Insights.MemoryProfiler.TimingView]
