[AudioInsights] - Support AudioModulationInsights within Unreal Insights standalone

#Synced-CL 51007057

diff --git a/Engine/Programs/UnrealInsights/Config/DefaultEngine.ini b/Engine/Programs/UnrealInsights/Config/DefaultEngine.ini
--- a/Engine/Programs/UnrealInsights/Config/DefaultEngine.ini
+++ b/Engine/Programs/UnrealInsights/Config/DefaultEngine.ini
@@ -1,5 +1,6 @@
 [Plugins]
 +ProgramEnabledPlugins=AudioInsights
++ProgramEnabledPlugins=AudioModulationInsights
 +ProgramEnabledPlugins=GameplayInsights
 +ProgramEnabledPlugins=SlateInsights
 +ProgramEnabledPlugins=RenderGraphInsights
